Added missing '{' to TP_ME_BlankBP.uproject, found building ShooterGame.

#rnx

diff --git a/Templates/TP_ME_BlankBP/TP_ME_BlankBP.uproject b/Templates/TP_ME_BlankBP/TP_ME_BlankBP.uproject
--- a/Templates/TP_ME_BlankBP/TP_ME_BlankBP.uproject
+++ b/Templates/TP_ME_BlankBP/TP_ME_BlankBP.uproject
@@ -31,6 +31,7 @@
 				"Mac"
 			]
 		},
+		{
 			"Name": "PythonScriptPlugin",
 			"Enabled": true
 		},
