Virtual Production Templates:  Use targeted platforms for windows 64- Ndiplay, Blackmagic, and AJA, Hap Media, Apple ProRes plug-ins

diff --git a/Templates/TP_ME_BlankBP/TP_ME_BlankBP.uproject b/Templates/TP_ME_BlankBP/TP_ME_BlankBP.uproject
--- a/Templates/TP_ME_BlankBP/TP_ME_BlankBP.uproject
+++ b/Templates/TP_ME_BlankBP/TP_ME_BlankBP.uproject
@@ -82,24 +82,39 @@
 		{
 			"Name": "nDisplay",
 			"Enabled": true
+			"SupportedTargetPlatforms": [
+				"Win64"
+			]
 		},
 		{
 			"Name": "AjaMedia",
 			"Enabled": true,
 			"MarketplaceURL": "com.epicgames.launcher://ue/marketplace/content/91088fb17fb04b6b9eda41aaf8fda21e"
+			"SupportedTargetPlatforms": [
+				"Win64"
+			]
 		},
 		{
 			"Name": "BlackmagicMedia",
 			"Enabled": true,
 			"MarketplaceURL": "com.epicgames.launcher://ue/marketplace/content/433d28bdfb764ec5b2e60c4f45d5cb5c"
+			"SupportedTargetPlatforms": [
+				"Win64"
+			]
 		},
 		{
 			"Name": "AppleProResMedia",
 			"Enabled": true
+			"SupportedTargetPlatforms": [
+				"Win64"
+			]
 		},
 		{
 			"Name": "HAPMedia",
 			"Enabled": true
+			"SupportedTargetPlatforms": [
+				"Win64"
+			]
 		},
 		{
 			"Name": "AppleARKit",
diff --git a/Templates/TP_ME_VProdBP/TP_ME_VprodBP.uproject b/Templates/TP_ME_VProdBP/TP_ME_VprodBP.uproject
--- a/Templates/TP_ME_VProdBP/TP_ME_VprodBP.uproject
+++ b/Templates/TP_ME_VProdBP/TP_ME_VprodBP.uproject
@@ -82,24 +82,39 @@
 		{
 			"Name": "nDisplay",
 			"Enabled": true
+			"SupportedTargetPlatforms": [
+				"Win64"
+			]
 		},
 		{
 			"Name": "AjaMedia",
 			"Enabled": true,
 			"MarketplaceURL": "com.epicgames.launcher://ue/marketplace/content/91088fb17fb04b6b9eda41aaf8fda21e"
+			"SupportedTargetPlatforms": [
+				"Win64"
+			]
 		},
 		{
 			"Name": "BlackmagicMedia",
 			"Enabled": true,
 			"MarketplaceURL": "com.epicgames.launcher://ue/marketplace/content/433d28bdfb764ec5b2e60c4f45d5cb5c"
+			"SupportedTargetPlatforms": [
+				"Win64"
+			]
 		},
 		{
 			"Name": "AppleProResMedia",
 			"Enabled": true
+			"SupportedTargetPlatforms": [
+				"Win64"
+			]
 		},
 		{
 			"Name": "HAPMedia",
 			"Enabled": true
+			"SupportedTargetPlatforms": [
+				"Win64"
+			]
 		},
 		{
 			"Name": "AppleARKit",
