Virtual Production Templates:  Missing comma after AJA and BM plug-ins added.

diff --git a/Templates/TP_ME_BlankBP/TP_ME_BlankBP.uproject b/Templates/TP_ME_BlankBP/TP_ME_BlankBP.uproject
--- a/Templates/TP_ME_BlankBP/TP_ME_BlankBP.uproject
+++ b/Templates/TP_ME_BlankBP/TP_ME_BlankBP.uproject
@@ -81,7 +81,7 @@
 		},
 		{
 			"Name": "nDisplay",
-			"Enabled": true
+			"Enabled": true,
 			"SupportedTargetPlatforms": [
 				"Win64"
 			]
@@ -89,7 +89,7 @@
 		{
 			"Name": "AjaMedia",
 			"Enabled": true,
-			"MarketplaceURL": "com.epicgames.launcher://ue/marketplace/content/91088fb17fb04b6b9eda41aaf8fda21e"
+			"MarketplaceURL": "com.epicgames.launcher://ue/marketplace/content/91088fb17fb04b6b9eda41aaf8fda21e",
 			"SupportedTargetPlatforms": [
 				"Win64"
 			]
@@ -97,21 +97,21 @@
 		{
 			"Name": "BlackmagicMedia",
 			"Enabled": true,
-			"MarketplaceURL": "com.epicgames.launcher://ue/marketplace/content/433d28bdfb764ec5b2e60c4f45d5cb5c"
+			"MarketplaceURL": "com.epicgames.launcher://ue/marketplace/content/433d28bdfb764ec5b2e60c4f45d5cb5c",
 			"SupportedTargetPlatforms": [
 				"Win64"
 			]
 		},
 		{
 			"Name": "AppleProResMedia",
-			"Enabled": true
+			"Enabled": true,
 			"SupportedTargetPlatforms": [
 				"Win64"
 			]
 		},
 		{
 			"Name": "HAPMedia",
-			"Enabled": true
+			"Enabled": true,
 			"SupportedTargetPlatforms": [
 				"Win64"
 			]
diff --git a/Templates/TP_ME_VProdBP/TP_ME_VprodBP.uproject b/Templates/TP_ME_VProdBP/TP_ME_VprodBP.uproject
--- a/Templates/TP_ME_VProdBP/TP_ME_VprodBP.uproject
+++ b/Templates/TP_ME_VProdBP/TP_ME_VprodBP.uproject
@@ -81,7 +81,7 @@
 		},
 		{
 			"Name": "nDisplay",
-			"Enabled": true
+			"Enabled": true,
 			"SupportedTargetPlatforms": [
 				"Win64"
 			]
@@ -89,7 +89,7 @@
 		{
 			"Name": "AjaMedia",
 			"Enabled": true,
-			"MarketplaceURL": "com.epicgames.launcher://ue/marketplace/content/91088fb17fb04b6b9eda41aaf8fda21e"
+			"MarketplaceURL": "com.epicgames.launcher://ue/marketplace/content/91088fb17fb04b6b9eda41aaf8fda21e",
 			"SupportedTargetPlatforms": [
 				"Win64"
 			]
@@ -97,21 +97,21 @@
 		{
 			"Name": "BlackmagicMedia",
 			"Enabled": true,
-			"MarketplaceURL": "com.epicgames.launcher://ue/marketplace/content/433d28bdfb764ec5b2e60c4f45d5cb5c"
+			"MarketplaceURL": "com.epicgames.launcher://ue/marketplace/content/433d28bdfb764ec5b2e60c4f45d5cb5c",
 			"SupportedTargetPlatforms": [
 				"Win64"
 			]
 		},
 		{
 			"Name": "AppleProResMedia",
-			"Enabled": true
+			"Enabled": true,
 			"SupportedTargetPlatforms": [
 				"Win64"
 			]
 		},
 		{
 			"Name": "HAPMedia",
-			"Enabled": true
+			"Enabled": true,
 			"SupportedTargetPlatforms": [
 				"Win64"
 			]
