ME template skeleton fix

enabling the right set of plugins

diff --git a/Templates/TP_ME_BlankBP/TP_ME_BlankBP.uproject b/Templates/TP_ME_BlankBP/TP_ME_BlankBP.uproject
--- a/Templates/TP_ME_BlankBP/TP_ME_BlankBP.uproject
+++ b/Templates/TP_ME_BlankBP/TP_ME_BlankBP.uproject
@@ -35,35 +35,35 @@
 		},
 		{
 			"Name": "SequencerScripting",
-			"Enabled": false
+			"Enabled": true
 		},
 		{
 			"Name": "Takes",
-			"Enabled": false
+			"Enabled": true
 		},
 		{
 			"Name": "VirtualProductionUtilities",
-			"Enabled": false
+			"Enabled": true
 		},
 		{
 			"Name": "VirtualCamera",
-			"Enabled": false
+			"Enabled": true
 		},
 		{
 			"Name": "LiveLink",
-			"Enabled": false
+			"Enabled": true
 		},
 		{
 			"Name": "Composure",
-			"Enabled": false
+			"Enabled": true
 		},
 		{
 			"Name": "MediaFrameworkUtilities",
-			"Enabled": false
+			"Enabled": true
 		},
 		{
 			"Name": "MediaIOFramework",
-			"Enabled": false
+			"Enabled": true
 		}
 	]
 }
\ No newline at end of file
