Remove reference to marketplace for AJA and Blackmagic plugin. They are now shipped with the engine.

diff --git a/Templates/TP_ME_BlankBP/TP_ME_BlankBP.uproject b/Templates/TP_ME_BlankBP/TP_ME_BlankBP.uproject
--- a/Templates/TP_ME_BlankBP/TP_ME_BlankBP.uproject
+++ b/Templates/TP_ME_BlankBP/TP_ME_BlankBP.uproject
@@ -89,7 +89,6 @@
 		{
 			"Name": "AjaMedia",
 			"Enabled": true,
-			"MarketplaceURL": "com.epicgames.launcher://ue/marketplace/content/91088fb17fb04b6b9eda41aaf8fda21e",
 			"SupportedTargetPlatforms": [
 				"Win64"
 			]
@@ -97,7 +96,6 @@
 		{
 			"Name": "BlackmagicMedia",
 			"Enabled": true,
-			"MarketplaceURL": "com.epicgames.launcher://ue/marketplace/content/433d28bdfb764ec5b2e60c4f45d5cb5c",
 			"SupportedTargetPlatforms": [
 				"Win64"
 			]
diff --git a/Templates/TP_ME_VProdBP/TP_ME_VprodBP.uproject b/Templates/TP_ME_VProdBP/TP_ME_VprodBP.uproject
--- a/Templates/TP_ME_VProdBP/TP_ME_VprodBP.uproject
+++ b/Templates/TP_ME_VProdBP/TP_ME_VprodBP.uproject
@@ -89,7 +89,6 @@
 		{
 			"Name": "AjaMedia",
 			"Enabled": true,
-			"MarketplaceURL": "com.epicgames.launcher://ue/marketplace/content/91088fb17fb04b6b9eda41aaf8fda21e",
 			"SupportedTargetPlatforms": [
 				"Win64"
 			]
@@ -97,7 +96,6 @@
 		{
 			"Name": "BlackmagicMedia",
 			"Enabled": true,
-			"MarketplaceURL": "com.epicgames.launcher://ue/marketplace/content/433d28bdfb764ec5b2e60c4f45d5cb5c",
 			"SupportedTargetPlatforms": [
 				"Win64"
 			]
