Enabed DMX Plugin in ME_BlankBP and ME_VProdBP Template projects #rb none #jira UE-89867

diff --git a/Templates/TP_ME_BlankBP/TP_ME_BlankBP.uproject b/Templates/TP_ME_BlankBP/TP_ME_BlankBP.uproject
--- a/Templates/TP_ME_BlankBP/TP_ME_BlankBP.uproject
+++ b/Templates/TP_ME_BlankBP/TP_ME_BlankBP.uproject
@@ -139,8 +139,10 @@
 		{
 			"Name": "VariantManager",
 			"Enabled": true
+		},
+		{
+			"Name": "DMXProtocol",
+			"Enabled": true
 		}
-		
-
 	]
 }
\ No newline at end of file
diff --git a/Templates/TP_ME_VProdBP/TP_ME_VprodBP.uproject b/Templates/TP_ME_VProdBP/TP_ME_VprodBP.uproject
--- a/Templates/TP_ME_VProdBP/TP_ME_VprodBP.uproject
+++ b/Templates/TP_ME_VProdBP/TP_ME_VprodBP.uproject
@@ -139,6 +139,10 @@
 		{
 			"Name": "VariantManager",
 			"Enabled": true
+		},
+		{
+			"Name": "DMXProtocol",
+			"Enabled": true
 		}
 	]
 }
\ No newline at end of file
