M&E Template Projects
Disabling Oculus/SteamVR plugins for OpenXR
Enabling Virtual Scouting

diff --git a/Templates/TP_InCamVFXBP/TP_InCamVFXBP.uproject b/Templates/TP_InCamVFXBP/TP_InCamVFXBP.uproject
--- a/Templates/TP_InCamVFXBP/TP_InCamVFXBP.uproject
+++ b/Templates/TP_InCamVFXBP/TP_InCamVFXBP.uproject
@@ -7,6 +7,23 @@
 		{
 			"Name": "ICVFX",
 			"Enabled": true
+		},
+		{
+			"Name": "OpenXR",
+			"Enabled": true,
+			"SupportedTargetPlatforms": [
+				"Win64",
+				"Linux",
+				"Android"
+			]
+		},
+		{
+			"Name": "VirtualScouting",
+			"Enabled": true,
+			"SupportedTargetPlatforms": [
+				"Win64",
+				"Linux"
+			]
 		}
 	],
 	"TargetPlatforms": [
diff --git a/Templates/TP_ME_BlankBP/TP_ME_BlankBP.uproject b/Templates/TP_ME_BlankBP/TP_ME_BlankBP.uproject
--- a/Templates/TP_ME_BlankBP/TP_ME_BlankBP.uproject
+++ b/Templates/TP_ME_BlankBP/TP_ME_BlankBP.uproject
@@ -5,19 +5,20 @@
 	"Description": "",
 	"Plugins": [
 		{
-			"Name": "OculusVR",
+			"Name": "OpenXR",
 			"Enabled": true,
 			"SupportedTargetPlatforms": [
-				"Win32",
-				"Win64"
+				"Win64",
+				"Linux",
+				"Android"
 			]
 		},
 		{
-			"Name": "SteamVR",
+			"Name": "VirtualScouting",
 			"Enabled": true,
 			"SupportedTargetPlatforms": [
-				"Win32",
-				"Win64"
+				"Win64",
+				"Linux"
 			]
 		},
 		{
diff --git a/Templates/TP_ME_VProdBP/TP_ME_VprodBP.uproject b/Templates/TP_ME_VProdBP/TP_ME_VprodBP.uproject
--- a/Templates/TP_ME_VProdBP/TP_ME_VprodBP.uproject
+++ b/Templates/TP_ME_VProdBP/TP_ME_VprodBP.uproject
@@ -5,19 +5,20 @@
 	"Description": "",
 	"Plugins": [
 		{
-			"Name": "OculusVR",
+			"Name": "OpenXR",
 			"Enabled": true,
 			"SupportedTargetPlatforms": [
-				"Win32",
-				"Win64"
+				"Win64",
+				"Linux",
+				"Android"
 			]
 		},
 		{
-			"Name": "SteamVR",
+			"Name": "VirtualScouting",
 			"Enabled": true,
 			"SupportedTargetPlatforms": [
-				"Win32",
-				"Win64"
+				"Win64",
+				"Linux"
 			]
 		},
 		{
diff --git a/Templates/TP_nDisplayBP/TP_nDisplayBP.uproject b/Templates/TP_nDisplayBP/TP_nDisplayBP.uproject
--- a/Templates/TP_nDisplayBP/TP_nDisplayBP.uproject
+++ b/Templates/TP_nDisplayBP/TP_nDisplayBP.uproject
@@ -13,17 +13,20 @@
 			]
 		},
 		{
-			"Name": "OculusVR",
-			"Enabled": false,
+			"Name": "OpenXR",
+			"Enabled": true,
 			"SupportedTargetPlatforms": [
-				"Win64"
+				"Win64",
+				"Linux",
+				"Android"
 			]
 		},
 		{
-			"Name": "SteamVR",
-			"Enabled": false,
+			"Name": "VirtualScouting",
+			"Enabled": true,
 			"SupportedTargetPlatforms": [
-				"Win64"
+				"Win64",
+				"Linux"
 			]
 		},
 		{
